Rename/move file(s) to fix the Missing Icon

diff --git a/Templates/TP_CollaborativeBP/TP_CollaborativeBP.uproject b/Templates/TP_CollaborativeBP/TP_CollaborativeBP.uproject
new file mode 100644
--- /dev/null
+++ b/Templates/TP_CollaborativeBP/TP_CollaborativeBP.uproject
@@ -0,0 +1,6 @@
+{
+	"FileVersion": 3,
+	"EngineAssociation": "",
+	"Category": "",
+	"Description": ""
+}
\ No newline at end of file
diff --git a/Templates/TP_CollaborativeBP/TP_Collaborative_V.uproject b/Templates/TP_CollaborativeBP/TP_Collaborative_V.uproject
deleted file mode 100644
