Collaborative Template : Fixes

Fixing the flashing sky.
Fix Missing Widows Collisions.
Fix Glass Window material for selected Postprocess.
Remove Scene Capture Camera from main Level
Fix Save Game Warning when play in editor
Fixe Tree LOD Issue

diff --git a/Templates/TP_CollaborativeBP/TP_CollaborativeBP.uproject b/Templates/TP_CollaborativeBP/TP_CollaborativeBP.uproject
--- a/Templates/TP_CollaborativeBP/TP_CollaborativeBP.uproject
+++ b/Templates/TP_CollaborativeBP/TP_CollaborativeBP.uproject
@@ -18,7 +18,10 @@
 		},
 		{
 			"Name": "DatasmithCADImporter",
-			"Enabled": true
+			"Enabled": true,
+			"SupportedTargetPlatforms": [
+				"Win64"
+			]
 		},
 		{
 			"Name": "VariantManager",
@@ -36,7 +39,7 @@
 			"Name": "SequencerScripting",
 			"Enabled": false
 		},
-				{
+		{
 			"Name": "StaticMeshEditorExtension",
 			"Enabled": true
 		}
