Collaborative Viewer - Feature pack is not accessible through Add Feature or Content Pack

diff --git a/Templates/TP_CollaborativeBP/manifest.json b/Templates/TP_CollaborativeBP/manifest.json
--- a/Templates/TP_CollaborativeBP/manifest.json
+++ b/Templates/TP_CollaborativeBP/manifest.json
@@ -17,7 +17,7 @@
 	"AssetTypes": [],
 	"SearchTags": [],
 	"ClassTypes": "",
-	"Category": "",
+	"Category": "BlueprintFeature",
 	"FocusAsset": "",
 	"Thumbnail": "TP_CollaborativeBP.png",
 	"Screenshots":
