Fixing list of plugins for arch viz and mfg projects.

Axf and Dtasmith GLTF plugins are now enabled by default in those project type.
No code changes, just plugins enabled

Verification done:

create a new project using each templates affected by this submission.
Each project launched correctly
Desired plugins (AXF and GLTF) are now enabled

diff --git a/Templates/TP_AEC_ArchvisBP/TP_AEC_ArchvisBP.uproject b/Templates/TP_AEC_ArchvisBP/TP_AEC_ArchvisBP.uproject
--- a/Templates/TP_AEC_ArchvisBP/TP_AEC_ArchvisBP.uproject
+++ b/Templates/TP_AEC_ArchvisBP/TP_AEC_ArchvisBP.uproject
@@ -49,6 +49,17 @@
 		{
 			"Name": "StaticMeshEditorExtension",
 			"Enabled": true
+		},
+		{
+			"Name": "DatasmithGLTFImporter",
+			"Enabled": true
+		},
+		{
+			"Name": "AxFImporter",
+			"Enabled": true,
+			"SupportedTargetPlatforms": [
+				"Win64"
+			]
 		}
 	]
 }
\ No newline at end of file
diff --git a/Templates/TP_AEC_BlankBP/TP_AEC_BlankBP.uproject b/Templates/TP_AEC_BlankBP/TP_AEC_BlankBP.uproject
--- a/Templates/TP_AEC_BlankBP/TP_AEC_BlankBP.uproject
+++ b/Templates/TP_AEC_BlankBP/TP_AEC_BlankBP.uproject
@@ -17,13 +17,6 @@
 			"Name": "DatasmithImporter",
 			"Enabled": true
 		},
-		{
-			"Name": "DatasmithCADImporter",
-			"Enabled": true,
-			"SupportedTargetPlatforms": [
-				"Win64"
-			]
-		},
 		{
 			"Name": "VariantManager",
 			"Enabled": true
@@ -40,9 +33,34 @@
 			"Name": "SequencerScripting",
 			"Enabled": false
 		},
-				{
+		{
+			"Name": "DatasmithCADImporter",
+			"Enabled": true,
+			"SupportedTargetPlatforms": [
+				"Win64"
+			]
+		},
+		{
+			"Name": "DatasmithIFCImporter",
+			"Enabled": true,
+			"SupportedTargetPlatforms": [
+				"Win64"
+			]
+		},
+		{
 			"Name": "StaticMeshEditorExtension",
 			"Enabled": true
+		},
+		{
+			"Name": "DatasmithGLTFImporter",
+			"Enabled": true
+		},
+		{
+			"Name": "AxFImporter",
+			"Enabled": true,
+			"SupportedTargetPlatforms": [
+				"Win64"
+			]
 		}
 	]
 }
\ No newline at end of file
diff --git a/Templates/TP_CollaborativeBP/TP_CollaborativeBP.uproject b/Templates/TP_CollaborativeBP/TP_CollaborativeBP.uproject
--- a/Templates/TP_CollaborativeBP/TP_CollaborativeBP.uproject
+++ b/Templates/TP_CollaborativeBP/TP_CollaborativeBP.uproject
@@ -5,19 +5,16 @@
 	"Description": "",
 	"Plugins": [
 		{
-			"Name": "DataPrepEditor",
+			"Name": "SunPosition",
 			"Enabled": true
 		},
 		{
-			"Name": "DatasmithImporter",
+			"Name": "DataPrepEditor",
 			"Enabled": true
 		},
 		{
-			"Name": "DatasmithCADImporter",
-			"Enabled": true,
-			"SupportedTargetPlatforms": [
-				"Win64"
-			]
+			"Name": "DatasmithImporter",
+			"Enabled": true
 		},
 		{
 			"Name": "VariantManager",
@@ -35,9 +32,34 @@
 			"Name": "SequencerScripting",
 			"Enabled": false
 		},
+		{
+			"Name": "DatasmithCADImporter",
+			"Enabled": true,
+			"SupportedTargetPlatforms": [
+				"Win64"
+			]
+		},
+		{
+			"Name": "DatasmithIFCImporter",
+			"Enabled": true,
+			"SupportedTargetPlatforms": [
+				"Win64"
+			]
+		},
 		{
 			"Name": "StaticMeshEditorExtension",
 			"Enabled": true
+		},
+		{
+			"Name": "DatasmithGLTFImporter",
+			"Enabled": true
+		},
+		{
+			"Name": "AxFImporter",
+			"Enabled": true,
+			"SupportedTargetPlatforms": [
+				"Win64"
+			]
 		}
 	]
 }
\ No newline at end of file
diff --git a/Templates/TP_PhotoStudioBP/TP_PhotoStudioBP.uproject b/Templates/TP_PhotoStudioBP/TP_PhotoStudioBP.uproject
--- a/Templates/TP_PhotoStudioBP/TP_PhotoStudioBP.uproject
+++ b/Templates/TP_PhotoStudioBP/TP_PhotoStudioBP.uproject
@@ -17,13 +17,6 @@
 			"Name": "DatasmithImporter",
 			"Enabled": true
 		},
-		{
-			"Name": "DatasmithCADImporter",
-			"Enabled": true,
-			"SupportedTargetPlatforms": [
-				"Win64"
-			]
-		},
 		{
 			"Name": "VariantManager",
 			"Enabled": true
@@ -38,11 +31,36 @@
 		},
 		{
 			"Name": "SequencerScripting",
-			"Enabled": true
+			"Enabled": false
+		},
+		{
+			"Name": "DatasmithCADImporter",
+			"Enabled": true,
+			"SupportedTargetPlatforms": [
+				"Win64"
+			]
+		},
+		{
+			"Name": "DatasmithIFCImporter",
+			"Enabled": true,
+			"SupportedTargetPlatforms": [
+				"Win64"
+			]
 		},
 		{
 			"Name": "StaticMeshEditorExtension",
 			"Enabled": true
+		},
+		{
+			"Name": "DatasmithGLTFImporter",
+			"Enabled": true
+		},
+		{
+			"Name": "AxFImporter",
+			"Enabled": true,
+			"SupportedTargetPlatforms": [
+				"Win64"
+			]
 		}
 	]
 }
\ No newline at end of file
