Adding Chinese translation #jira ueloc-4711 #fyi michelle.zhao, alexis.zhang, yali.he

diff --git a/Templates/TP_CollaborativeBP/manifest.json b/Templates/TP_CollaborativeBP/manifest.json
--- a/Templates/TP_CollaborativeBP/manifest.json
+++ b/Templates/TP_CollaborativeBP/manifest.json
@@ -5,6 +5,10 @@
 		{
 			"Language": "en",
 			"Text": "Collab Viewer"
+		},
+		{
+			"Language": "zh-Hans",
+			"Text": "协作查看器"
 		}
 	],
 	"Description":
@@ -13,6 +17,11 @@
 			"Language": "en",
 			"Text": "This project template provides navigation and interaction for various platforms (Desktop, VR). Use as a starting point for collaborative Industrial or Architectural projects."
 
+		},
+		{
+			"Language": "zh-Hans",
+			"Text": "该项目模板提供面向多种平台（台式电脑，VR等）的导航及互动功能，适合用于协作式工业或建筑项目。"
+
 		},
 		{
 
