Removed references to deleted plugin.

diff --git a/Templates/TP_AEC_ArchvisBP/TP_AEC_ArchvisBP.uproject b/Templates/TP_AEC_ArchvisBP/TP_AEC_ArchvisBP.uproject
--- a/Templates/TP_AEC_ArchvisBP/TP_AEC_ArchvisBP.uproject
+++ b/Templates/TP_AEC_ArchvisBP/TP_AEC_ArchvisBP.uproject
@@ -46,10 +46,6 @@
 				"Win64"
 			]
 		},
-		{
-			"Name": "StaticMeshEditorExtension",
-			"Enabled": true
-		},
 		{
 			"Name": "DatasmithGLTFImporter",
 			"Enabled": true
diff --git a/Templates/TP_AEC_BlankBP/TP_AEC_BlankBP.uproject b/Templates/TP_AEC_BlankBP/TP_AEC_BlankBP.uproject
--- a/Templates/TP_AEC_BlankBP/TP_AEC_BlankBP.uproject
+++ b/Templates/TP_AEC_BlankBP/TP_AEC_BlankBP.uproject
@@ -47,10 +47,6 @@
 				"Win64"
 			]
 		},
-		{
-			"Name": "StaticMeshEditorExtension",
-			"Enabled": true
-		},
 		{
 			"Name": "DatasmithGLTFImporter",
 			"Enabled": true
diff --git a/Templates/TP_CollaborativeBP/TP_CollaborativeBP.uproject b/Templates/TP_CollaborativeBP/TP_CollaborativeBP.uproject
--- a/Templates/TP_CollaborativeBP/TP_CollaborativeBP.uproject
+++ b/Templates/TP_CollaborativeBP/TP_CollaborativeBP.uproject
@@ -46,10 +46,6 @@
 				"Win64"
 			]
 		},
-		{
-			"Name": "StaticMeshEditorExtension",
-			"Enabled": true
-		},
 		{
 			"Name": "DatasmithGLTFImporter",
 			"Enabled": true
diff --git a/Templates/TP_HololensBP/TP_HololensBP.uproject b/Templates/TP_HololensBP/TP_HololensBP.uproject
--- a/Templates/TP_HololensBP/TP_HololensBP.uproject
+++ b/Templates/TP_HololensBP/TP_HololensBP.uproject
@@ -46,10 +46,6 @@
 				"Win64"
 			]
 		},
-		{
-			"Name": "StaticMeshEditorExtension",
-			"Enabled": true
-		},
 		{
 			"Name": "DatasmithGLTFImporter",
 			"Enabled": true
diff --git a/Templates/TP_PhotoStudioBP/TP_PhotoStudioBP.uproject b/Templates/TP_PhotoStudioBP/TP_PhotoStudioBP.uproject
--- a/Templates/TP_PhotoStudioBP/TP_PhotoStudioBP.uproject
+++ b/Templates/TP_PhotoStudioBP/TP_PhotoStudioBP.uproject
@@ -47,10 +47,6 @@
 				"Win64"
 			]
 		},
-		{
-			"Name": "StaticMeshEditorExtension",
-			"Enabled": true
-		},
 		{
 			"Name": "DatasmithGLTFImporter",
 			"Enabled": true
